Adding Sun and Sky blueprint to Lighting modes pane.

diff --git a/Engine/Plugins/Runtime/SunPosition/Source/SunPosition/Private/SunPositionModule.cpp b/Engine/Plugins/Runtime/SunPosition/Source/SunPosition/Private/SunPositionModule.cpp
--- a/Engine/Plugins/Runtime/SunPosition/Source/SunPosition/Private/SunPositionModule.cpp
+++ b/Engine/Plugins/Runtime/SunPosition/Source/SunPosition/Private/SunPositionModule.cpp
@@ -1,12 +1,18 @@
 // Copyright 1998-2019 Epic Games, Inc. All Rights Reserved.
 
 #include "SunPositionModule.h"
+#include "SunPositionPlacement.h"
+#include "SunPositionStyle.h"
 
-#include "Modules/ModuleInterface.h"
-#include "Modules/ModuleManager.h"
-
-IMPLEMENT_MODULE(FDefaultModuleImpl, SunPosition)
-
-
+void FSunPositionModule::StartupModule()
+{
+	FSunPositionStyle::Initialize();
+	FSunPositionPlacement::RegisterPlacement();
+}
 
+void FSunPositionModule::ShutdownModule()
+{
+	FSunPositionStyle::Shutdown();
+}
 
+IMPLEMENT_MODULE(FSunPositionModule, SunPosition)
\ No newline at end of file
diff --git a/Engine/Plugins/Runtime/SunPosition/Source/SunPosition/Private/SunPositionModule.h b/Engine/Plugins/Runtime/SunPosition/Source/SunPosition/Private/SunPositionModule.h
--- a/Engine/Plugins/Runtime/SunPosition/Source/SunPosition/Private/SunPositionModule.h
+++ b/Engine/Plugins/Runtime/SunPosition/Source/SunPosition/Private/SunPositionModule.h
@@ -2,4 +2,13 @@
 #pragma once
 
 #include "CoreMinimal.h"
+#include "Modules/ModuleManager.h"
 
+class FSunPositionModule : public IModuleInterface
+{
+public:
+
+	/** IModuleInterface implementation */
+	virtual void StartupModule() override;
+	virtual void ShutdownModule() override;
+};
\ No newline at end of file
diff --git a/Engine/Plugins/Runtime/SunPosition/Source/SunPosition/Private/SunPositionPlacement.cpp b/Engine/Plugins/Runtime/SunPosition/Source/SunPosition/Private/SunPositionPlacement.cpp
new file mode 100644
--- /dev/null
+++ b/Engine/Plugins/Runtime/SunPosition/Source/SunPosition/Private/SunPositionPlacement.cpp
@@ -0,0 +1,28 @@
+// Copyright 1998-2019 Epic Games, Inc. All Rights Reserved.
+
+#include "SunPositionPlacement.h"
+#include "IPlacementModeModule.h"
+#include "ActorFactories/ActorFactoryBlueprint.h"
+
+void FSunPositionPlacement::RegisterPlacement()
+{
+	UBlueprint* SunSky = Cast<UBlueprint>(FSoftObjectPath(TEXT("/SunPosition/SunSky.SunSky")).TryLoad());
+	if (SunSky == nullptr)
+	{
+		return;
+	}
+
+	IPlacementModeModule& PlacementModeModule = IPlacementModeModule::Get();
+	FPlacementCategoryInfo Info = *PlacementModeModule.GetRegisteredPlacementCategory(FBuiltInPlacementCategories::Lights());
+
+	FPlaceableItem* BPPlacement = new FPlaceableItem(
+		*UActorFactoryBlueprint::StaticClass(),
+		FAssetData(SunSky, true),
+		FName("SunPosition.ModesIcon"),
+		TOptional<FLinearColor>(),
+		TOptional<int32>(),
+		NSLOCTEXT("PlacementMode", "Sun and Sky", "Sun and Sky")
+	);
+
+	IPlacementModeModule::Get().RegisterPlaceableItem(Info.UniqueHandle, MakeShareable(BPPlacement));
+}
\ No newline at end of file
diff --git a/Engine/Plugins/Runtime/SunPosition/Source/SunPosition/Private/SunPositionStyle.cpp b/Engine/Plugins/Runtime/SunPosition/Source/SunPosition/Private/SunPositionStyle.cpp
new file mode 100644
--- /dev/null
+++ b/Engine/Plugins/Runtime/SunPosition/Source/SunPosition/Private/SunPositionStyle.cpp
@@ -0,0 +1,37 @@
+// Copyright 1998-2019 Epic Games, Inc. All Rights Reserved.
+
+#include "SunPositionStyle.h"
+#include "Styling/SlateStyle.h"
+#include "Styling/SlateStyleRegistry.h"
+#include "Styling/SlateTypes.h"
+#include "Interfaces/IPluginManager.h"
+
+TSharedPtr<FSlateStyleSet> FSunPositionStyle::StyleSet;
+
+void FSunPositionStyle::Initialize()
+{
+	if (StyleSet.IsValid())
+	{
+		return;
+	}
+
+	StyleSet = MakeShared<FSlateStyleSet>(FName("SunPositionStyle"));
+
+	StyleSet->SetContentRoot(FPaths::EngineContentDir() / TEXT("Editor/Slate"));
+	StyleSet->SetCoreContentRoot(FPaths::EngineContentDir() / TEXT("Slate"));
+
+	FString IconPath = IPluginManager::Get().FindPlugin(TEXT("SunPosition"))->GetBaseDir() + TEXT("/Resources/SunPosition.png");
+	StyleSet->Set("SunPosition.ModesIcon", new FSlateImageBrush(IconPath, FVector2D(40.0f, 40.0f)));
+
+	FSlateStyleRegistry::RegisterSlateStyle(*StyleSet.Get());
+}
+
+void FSunPositionStyle::Shutdown()
+{
+	if (StyleSet.IsValid())
+	{
+		FSlateStyleRegistry::UnRegisterSlateStyle(*StyleSet.Get());
+		ensure(StyleSet.IsUnique());
+		StyleSet.Reset();
+	}
+}
\ No newline at end of file
diff --git a/Engine/Plugins/Runtime/SunPosition/Source/SunPosition/Public/SunPositionPlacement.h b/Engine/Plugins/Runtime/SunPosition/Source/SunPosition/Public/SunPositionPlacement.h
new file mode 100644
--- /dev/null
+++ b/Engine/Plugins/Runtime/SunPosition/Source/SunPosition/Public/SunPositionPlacement.h
@@ -0,0 +1,11 @@
+// Copyright 1998-2019 Epic Games, Inc. All Rights Reserved.
+
+#pragma once
+
+#include "CoreMinimal.h"
+
+class FSunPositionPlacement
+{
+public:
+	static void RegisterPlacement();
+};
\ No newline at end of file
diff --git a/Engine/Plugins/Runtime/SunPosition/Source/SunPosition/Public/SunPositionStyle.h b/Engine/Plugins/Runtime/SunPosition/Source/SunPosition/Public/SunPositionStyle.h
new file mode 100644
--- /dev/null
+++ b/Engine/Plugins/Runtime/SunPosition/Source/SunPosition/Public/SunPositionStyle.h
@@ -0,0 +1,17 @@
+// Copyright 1998-2019 Epic Games, Inc. All Rights Reserved.
+
+#pragma once
+
+#include "CoreMinimal.h"
+#include "Styling/SlateStyle.h"
+
+class FSunPositionStyle
+{
+public:
+	static void Initialize();
+	static void Shutdown();
+	static TSharedPtr<ISlateStyle> Get() { return StyleSet; }
+
+private:
+	static TSharedPtr<FSlateStyleSet> StyleSet;
+};
\ No newline at end of file
diff --git a/Engine/Plugins/Runtime/SunPosition/Source/SunPosition/SunPosition.build.cs b/Engine/Plugins/Runtime/SunPosition/Source/SunPosition/SunPosition.build.cs
--- a/Engine/Plugins/Runtime/SunPosition/Source/SunPosition/SunPosition.build.cs
+++ b/Engine/Plugins/Runtime/SunPosition/Source/SunPosition/SunPosition.build.cs
@@ -11,6 +11,11 @@ namespace UnrealBuildTool.Rules
                     "Core",
                     "CoreUObject",
                     "Engine",
+					"Projects",
+                    "Slate",
+					"SlateCore",
+                    "UnrealEd",
+                    "PlacementMode",
                 }
             );
 		}
diff --git a/Engine/Plugins/Runtime/SunPosition/SunPosition.uplugin b/Engine/Plugins/Runtime/SunPosition/SunPosition.uplugin
--- a/Engine/Plugins/Runtime/SunPosition/SunPosition.uplugin
+++ b/Engine/Plugins/Runtime/SunPosition/SunPosition.uplugin
@@ -18,7 +18,7 @@
 		{
 			"Name": "SunPosition",
 			"Type": "Runtime",
-			"LoadingPhase": "Default"
+			"LoadingPhase": "PostEngineInit"
 		}
 	]
 }
\ No newline at end of file
