Fix sun position plugin creating dependency on the editor during runtime compile

diff --git a/Engine/Plugins/Runtime/SunPosition/Source/SunPosition/Private/SunPositionPlacement.cpp b/Engine/Plugins/Runtime/SunPosition/Source/SunPosition/Private/SunPositionPlacement.cpp
--- a/Engine/Plugins/Runtime/SunPosition/Source/SunPosition/Private/SunPositionPlacement.cpp
+++ b/Engine/Plugins/Runtime/SunPosition/Source/SunPosition/Private/SunPositionPlacement.cpp
@@ -1,28 +1,35 @@
 // Copyright 1998-2019 Epic Games, Inc. All Rights Reserved.
 
 #include "SunPositionPlacement.h"
+#if WITH_EDITOR
 #include "IPlacementModeModule.h"
 #include "ActorFactories/ActorFactoryBlueprint.h"
+#endif // WITH_EDITOR
 
 void FSunPositionPlacement::RegisterPlacement()
 {
-	UBlueprint* SunSky = Cast<UBlueprint>(FSoftObjectPath(TEXT("/SunPosition/SunSky.SunSky")).TryLoad());
-	if (SunSky == nullptr)
+#if WITH_EDITOR
+	if (GIsEditor)
 	{
-		return;
-	}
+		UBlueprint* SunSky = Cast<UBlueprint>(FSoftObjectPath(TEXT("/SunPosition/SunSky.SunSky")).TryLoad());
+		if (SunSky == nullptr)
+		{
+			return;
+		}
 
-	IPlacementModeModule& PlacementModeModule = IPlacementModeModule::Get();
-	FPlacementCategoryInfo Info = *PlacementModeModule.GetRegisteredPlacementCategory(FBuiltInPlacementCategories::Lights());
+		IPlacementModeModule& PlacementModeModule = IPlacementModeModule::Get();
+		FPlacementCategoryInfo Info = *PlacementModeModule.GetRegisteredPlacementCategory(FBuiltInPlacementCategories::Lights());
 
-	FPlaceableItem* BPPlacement = new FPlaceableItem(
-		*UActorFactoryBlueprint::StaticClass(),
-		FAssetData(SunSky, true),
-		FName("SunPosition.ModesIcon"),
-		TOptional<FLinearColor>(),
-		TOptional<int32>(),
-		NSLOCTEXT("PlacementMode", "Sun and Sky", "Sun and Sky")
-	);
+		FPlaceableItem* BPPlacement = new FPlaceableItem(
+			*UActorFactoryBlueprint::StaticClass(),
+			FAssetData(SunSky, true),
+			FName("SunPosition.ModesIcon"),
+			TOptional<FLinearColor>(),
+			TOptional<int32>(),
+			NSLOCTEXT("PlacementMode", "Sun and Sky", "Sun and Sky")
+		);
 
-	IPlacementModeModule::Get().RegisterPlaceableItem(Info.UniqueHandle, MakeShareable(BPPlacement));
+		IPlacementModeModule::Get().RegisterPlaceableItem(Info.UniqueHandle, MakeShareable(BPPlacement));
+	}
+#endif // WITH_EDITOR
 }
\ No newline at end of file
diff --git a/Engine/Plugins/Runtime/SunPosition/Source/SunPosition/SunPosition.build.cs b/Engine/Plugins/Runtime/SunPosition/Source/SunPosition/SunPosition.build.cs
--- a/Engine/Plugins/Runtime/SunPosition/Source/SunPosition/SunPosition.build.cs
+++ b/Engine/Plugins/Runtime/SunPosition/Source/SunPosition/SunPosition.build.cs
@@ -14,10 +14,18 @@ namespace UnrealBuildTool.Rules
 					"Projects",
                     "Slate",
 					"SlateCore",
-                    "UnrealEd",
-                    "PlacementMode",
                 }
             );
+
+			if (Target.bBuildEditor)
+			{
+				PrivateDependencyModuleNames.AddRange(
+					new string[] {
+						"UnrealEd",
+						"PlacementMode",
+					}
+				);
+			}
 		}
 	}
 }
