PR #8039: Revise the variable error when testing SunPosition (Contributed by [REDACTED_PERSON])

diff --git a/Engine/Plugins/Runtime/SunPosition/Source/SunPosition/Private/SunPosition.cpp b/Engine/Plugins/Runtime/SunPosition/Source/SunPosition/Private/SunPosition.cpp
--- a/Engine/Plugins/Runtime/SunPosition/Source/SunPosition/Private/SunPosition.cpp
+++ b/Engine/Plugins/Runtime/SunPosition/Source/SunPosition/Private/SunPosition.cpp
@@ -174,7 +174,7 @@ void TestExpectedValues(float Latitude, float Longitude, float TimeZone, bool bI
 
 	// Test azimuth values
 	float AzError = FMath::Abs(SunPositionData.Azimuth - ExpectedAzimuth)/ ExpectedAzimuth * 100.0f;
-	if (ElevError > AcceptableError)
+	if (AzError > AcceptableError)
 	{
 		UE_LOG(LogSunPosition, Warning, TEXT("Azimuth error above threshold %0.2f%%"), AzError);
 	}
@@ -209,4 +209,4 @@ void TestCalculator()
 		TestExpectedValues(45.0f, -73.0f, -5.0f, false, 2018, 9, 1, 5,  0,  0,  74.982,   -3.237,  -3.135);
 		TestExpectedValues(45.0f, -73.0f, -5.0f, false, 2018, 9, 1, 5, 30,  0,  80.308,    1.937,   2.225);
 	}
-}
\ No newline at end of file
+}
