Cherry pick Fix SunPosition initialization in UE5/Main to fix DDC issues

diff --git a/Engine/Plugins/Runtime/SunPosition/Source/SunPosition/Private/SunPosition.cpp b/Engine/Plugins/Runtime/SunPosition/Source/SunPosition/Private/SunPosition.cpp
--- a/Engine/Plugins/Runtime/SunPosition/Source/SunPosition/Private/SunPosition.cpp
+++ b/Engine/Plugins/Runtime/SunPosition/Source/SunPosition/Private/SunPosition.cpp
@@ -210,3 +210,11 @@ void TestCalculator()
 		TestExpectedValues(45.0f, -73.0f, -5.0f, false, 2018, 9, 1, 5, 30,  0,  80.308,    1.937,   2.225);
 	}
 }
+
+FSunPositionData::FSunPositionData()
+	: Elevation(0.0f)
+	, CorrectedElevation(0.0f)
+	, Azimuth(0.0f)
+{
+
+}
diff --git a/Engine/Plugins/Runtime/SunPosition/Source/SunPosition/Public/SunPosition.h b/Engine/Plugins/Runtime/SunPosition/Source/SunPosition/Public/SunPosition.h
--- a/Engine/Plugins/Runtime/SunPosition/Source/SunPosition/Public/SunPosition.h
+++ b/Engine/Plugins/Runtime/SunPosition/Source/SunPosition/Public/SunPosition.h
@@ -16,6 +16,8 @@ struct SUNPOSITION_API FSunPositionData
 
 public:
 
+	FSunPositionData();
+
 	/** Sun Elevation */
 	UPROPERTY(EditAnywhere, BlueprintReadWrite, Category="Sun Position")
 	float Elevation;
